Fix for Android Permission plugin headers not being generated for Lumin Builds

diff --git a/Engine/Plugins/Runtime/AndroidPermission/AndroidPermission.uplugin b/Engine/Plugins/Runtime/AndroidPermission/AndroidPermission.uplugin
--- a/Engine/Plugins/Runtime/AndroidPermission/AndroidPermission.uplugin
+++ b/Engine/Plugins/Runtime/AndroidPermission/AndroidPermission.uplugin
@@ -15,7 +15,8 @@
 	"IsBetaVersion": true,
 	"Installed": false,
 	"SupportedTargetPlatforms": [
-		"Android"
+		"Android",
+		"Lumin"
 	],
 	"Modules": [
 		{
