Removed platform restriction for AndroidPermission.uplugin. Plugin has Blueprint components and needs to be compiled for all platforms.

diff --git a/Engine/Plugins/Runtime/AndroidPermission/AndroidPermission.uplugin b/Engine/Plugins/Runtime/AndroidPermission/AndroidPermission.uplugin
--- a/Engine/Plugins/Runtime/AndroidPermission/AndroidPermission.uplugin
+++ b/Engine/Plugins/Runtime/AndroidPermission/AndroidPermission.uplugin
@@ -14,10 +14,6 @@
 	"CanContainContent": false,
 	"IsBetaVersion": true,
 	"Installed": false,
-	"SupportedTargetPlatforms": [
-		"Android",
-		"Lumin"
-	],
 	"Modules": [
 		{
 			"Name": "AndroidPermission",
