Fix broadcast of Android permission results to happen on GameThread

#android

diff --git a/Engine/Plugins/Runtime/AndroidPermission/Source/AndroidPermission/Private/AndroidPermissionCallbackProxy.cpp b/Engine/Plugins/Runtime/AndroidPermission/Source/AndroidPermission/Private/AndroidPermissionCallbackProxy.cpp
--- a/Engine/Plugins/Runtime/AndroidPermission/Source/AndroidPermission/Private/AndroidPermissionCallbackProxy.cpp
+++ b/Engine/Plugins/Runtime/AndroidPermission/Source/AndroidPermission/Private/AndroidPermissionCallbackProxy.cpp
@@ -41,8 +41,15 @@ JNI_METHOD void Java_com_google_vr_sdk_samples_permission_PermissionHelper_onAcq
 	UE_LOG(LogAndroidPermission, Log, TEXT("PermissionHelper_onAcquirePermissions %s %d (%d), Broadcasting..."),
 		*(arrPermissions[0]), arrGranted[0], num);
 
-	pProxy->OnPermissionsGrantedDelegate.Broadcast(arrPermissions, arrGranted);
-	pProxy->OnPermissionsGrantedDynamicDelegate.Broadcast(arrPermissions, arrGranted);
+	if (FTaskGraphInterface::IsRunning())
+	{
+		FGraphEventRef AquirePermissionsTask = FFunctionGraphTask::CreateAndDispatchWhenReady([arrPermissions, arrGranted]()
+		{
+			// pProxy has to be valid if got here (never reset after set)
+			pProxy->OnPermissionsGrantedDelegate.Broadcast(arrPermissions, arrGranted);
+			pProxy->OnPermissionsGrantedDynamicDelegate.Broadcast(arrPermissions, arrGranted);
+		}, TStatId(), NULL, ENamedThreads::GameThread);
+	}
 }
 #endif
 
